Fixed crash when CUDA isn't available during shutdown

#rnx

diff --git a/Engine/Plugins/Experimental/OptiXDenoise/Source/OptiXDenoise/Private/OptiXDenoiser.cpp b/Engine/Plugins/Experimental/OptiXDenoise/Source/OptiXDenoise/Private/OptiXDenoiser.cpp
--- a/Engine/Plugins/Experimental/OptiXDenoise/Source/OptiXDenoise/Private/OptiXDenoiser.cpp
+++ b/Engine/Plugins/Experimental/OptiXDenoise/Source/OptiXDenoise/Private/OptiXDenoiser.cpp
@@ -752,15 +752,19 @@ namespace UE::OptiXDenoiser
 			OptiXDenoiseContext->Destroy();
 		}
 
-		FCUDAModule::CUDA().cuCtxPushCurrent(FModuleManager::GetModuleChecked<FCUDAModule>("CUDA").GetCudaContext());
+		FCUDAModule& CudaModule = FModuleManager::GetModuleChecked<FCUDAModule>("CUDA");
+		if (CudaModule.IsAvailable())
+		{
+			FCUDAModule::CUDA().cuCtxPushCurrent(CudaModule.GetCudaContext());
 
-		CUDA_CHECK(FCUDAModule::CUDA().cuMemFree(Intensity));
-		CUDA_CHECK(FCUDAModule::CUDA().cuMemFree(AverageColor));
-		CUDA_CHECK(FCUDAModule::CUDA().cuMemFree(CudaScratch));
-		CUDA_CHECK(FCUDAModule::CUDA().cuMemFree(CudaState));
-		CUDA_CHECK(FCUDAModule::CUDA().cuMemFree(InternalMemIn));
-		CUDA_CHECK(FCUDAModule::CUDA().cuMemFree(InternalMemOut));
+			CUDA_CHECK(FCUDAModule::CUDA().cuMemFree(Intensity));
+			CUDA_CHECK(FCUDAModule::CUDA().cuMemFree(AverageColor));
+			CUDA_CHECK(FCUDAModule::CUDA().cuMemFree(CudaScratch));
+			CUDA_CHECK(FCUDAModule::CUDA().cuMemFree(CudaState));
+			CUDA_CHECK(FCUDAModule::CUDA().cuMemFree(InternalMemIn));
+			CUDA_CHECK(FCUDAModule::CUDA().cuMemFree(InternalMemOut));
 
-		FCUDAModule::CUDA().cuCtxPopCurrent(NULL);
+			FCUDAModule::CUDA().cuCtxPopCurrent(NULL);
+		}
 	}
 }
