Fix  Static Analysis warning For OptiX Denoiser.
#rnx

diff --git a/Engine/Plugins/Experimental/OptiXDenoise/Source/OptiXDenoise/Private/OptiXDenoiser.cpp b/Engine/Plugins/Experimental/OptiXDenoise/Source/OptiXDenoise/Private/OptiXDenoiser.cpp
--- a/Engine/Plugins/Experimental/OptiXDenoise/Source/OptiXDenoise/Private/OptiXDenoiser.cpp
+++ b/Engine/Plugins/Experimental/OptiXDenoise/Source/OptiXDenoise/Private/OptiXDenoiser.cpp
@@ -67,7 +67,7 @@ namespace UE::OptiXDenoiser
 
 	FOptiXImageData& FOptiXImage2D::GetRawImage()
 	{
-		check(Format == EOptiXImageFormat::CUDA_A32B32G32R32_F || EOptiXImageFormat::CUDA_FloatRGBA);
+		check(Format == EOptiXImageFormat::CUDA_A32B32G32R32_F || Format == EOptiXImageFormat::CUDA_FloatRGBA);
 
 		uint32_t PixelSize = Format == EOptiXImageFormat::CUDA_A32B32G32R32_F ? (4 * static_cast<uint32_t>(sizeof(float))) : (4 * static_cast<uint32_t>(sizeof(uint16_t)));
 		uint32_t RowStrideInBytes = PixelSize * Width;
